Copy of 11178294 from //UE4/Dev-Enterprise:

RemoteImport: Fix initialization check unhappy with varring initialization.

diff --git a/Engine/Source/Runtime/RemoteImportMessaging/Public/RemoteImportMessages.h b/Engine/Source/Runtime/RemoteImportMessaging/Public/RemoteImportMessages.h
--- a/Engine/Source/Runtime/RemoteImportMessaging/Public/RemoteImportMessages.h
+++ b/Engine/Source/Runtime/RemoteImportMessaging/Public/RemoteImportMessages.h
@@ -62,12 +62,7 @@ struct FRequestMessage
 	GENERATED_BODY()
 
 	UPROPERTY()
-	FGuid RequestID;
-
-	FRequestMessage()
-		: RequestID(FGuid::NewGuid())
-	{
-	}
+	int32 RequestID = -1;
 };
 
 /**
@@ -79,10 +74,10 @@ struct FResponseMessage
 	GENERATED_BODY()
 
 	UPROPERTY()
-	FGuid SourceRequestID;
+	int32 SourceRequestID;
 
 	FResponseMessage(const FRequestMessage* SourceRequest=nullptr)
-		: SourceRequestID(SourceRequest ? SourceRequest->RequestID : FGuid{})
+		: SourceRequestID(SourceRequest ? SourceRequest->RequestID : -1)
 	{
 	}
 };
